[Interchange - Assets] Re-updated assets with Latest changes.

#rnx

#changelist validated
#virtualized

diff --git a/Engine/Plugins/Interchange/Assets/Config/BaseInterchangeAssets.ini b/Engine/Plugins/Interchange/Assets/Config/BaseInterchangeAssets.ini
--- a/Engine/Plugins/Interchange/Assets/Config/BaseInterchangeAssets.ini
+++ b/Engine/Plugins/Interchange/Assets/Config/BaseInterchangeAssets.ini
@@ -1,6 +1,6 @@
 [CoreRedirects]
 
-; Redirectors for /Interchange/Functions/ (30) (29 - MX_Place2D)
+; Redirectors for /Interchange/Functions/ (30)
 +PackageRedirects=(OldName="/Interchange/Functions/MF_Iridescence.MF_Iridescence", NewName="/InterchangeAssets/Functions/MF_Iridescence.MF_Iridescence")
 +PackageRedirects=(OldName="/Interchange/Functions/MF_OrenNayerView.MF_OrenNayerView", NewName="/InterchangeAssets/Functions/MF_OrenNayerView.MF_OrenNayerView")
 +PackageRedirects=(OldName="/Interchange/Functions/MF_PhongToMetalRoughness.MF_PhongToMetalRoughness", NewName="/InterchangeAssets/Functions/MF_PhongToMetalRoughness.MF_PhongToMetalRoughness")
@@ -60,7 +60,7 @@
 +PackageRedirects=(OldName="/Interchange/gltf/MaterialBodies/MF_Unlit_Body.MF_Unlit_Body", NewName="/InterchangeAssets/gltf/MaterialBodies/MF_Unlit_Body.MF_Unlit_Body")
 
 
-; Redirectors for /Interchange/gltf/MaterialFunctions (23) (22 - MF_SpecularAnisotropy)
+; Redirectors for /Interchange/gltf/MaterialFunctions (23)
 +PackageRedirects=(OldName="/Interchange/gltf/MaterialFunctions/MF_BaseColor.MF_BaseColor", NewName="/InterchangeAssets/gltf/MaterialFunctions/MF_BaseColor.MF_BaseColor")
 +PackageRedirects=(OldName="/Interchange/gltf/MaterialFunctions/MF_Clearcoat.MF_Clearcoat", NewName="/InterchangeAssets/gltf/MaterialFunctions/MF_Clearcoat.MF_Clearcoat")
 +PackageRedirects=(OldName="/Interchange/gltf/MaterialFunctions/MF_DiffuseSpecGloss.MF_DiffuseSpecGloss", NewName="/InterchangeAssets/gltf/MaterialFunctions/MF_DiffuseSpecGloss.MF_DiffuseSpecGloss")
@@ -131,7 +131,7 @@
 +PackageRedirects=(OldName="/Interchange/gltf/MaterialInstances/MI_Unlit_Opaque_DS.MI_Unlit_Opaque_DS", NewName="/InterchangeAssets/gltf/MaterialInstances/MI_Unlit_Opaque_DS.MI_Unlit_Opaque_DS")
 
 
-; Redirectors for /Interchange/gltf/Textures (6) (5 - T_Anisotropy_Direction_Linear)
+; Redirectors for /Interchange/gltf/Textures (6)
 +PackageRedirects=(OldName="/Interchange/gltf/Textures/T_Anisotropy_Direction_Linear.T_Anisotropy_Direction_Linear", NewName="/InterchangeAssets/gltf/Textures/T_Anisotropy_Direction_Linear.T_Anisotropy_Direction_Linear")
 +PackageRedirects=(OldName="/Interchange/gltf/Textures/T_Black_srgb.T_Black_srgb", NewName="/InterchangeAssets/gltf/Textures/T_Black_srgb.T_Black_srgb")
 +PackageRedirects=(OldName="/Interchange/gltf/Textures/T_Generic_N.T_Generic_N", NewName="/InterchangeAssets/gltf/Textures/T_Generic_N.T_Generic_N")
diff --git a/Engine/Plugins/Interchange/Runtime/Source/Import/Private/MaterialX/MaterialXUtils/MaterialXManager.cpp b/Engine/Plugins/Interchange/Runtime/Source/Import/Private/MaterialX/MaterialXUtils/MaterialXManager.cpp
--- a/Engine/Plugins/Interchange/Runtime/Source/Import/Private/MaterialX/MaterialXUtils/MaterialXManager.cpp
+++ b/Engine/Plugins/Interchange/Runtime/Source/Import/Private/MaterialX/MaterialXUtils/MaterialXManager.cpp
@@ -385,7 +385,7 @@ namespace UE::Interchange::MaterialX
 															    TEXT("MaterialFunction'/InterchangeAssets/Functions/MX_Artistic_IOR.MX_Artistic_IOR'"),
 															    TEXT("MaterialFunction'/InterchangeAssets/Functions/MX_Roughness_Anisotropy.MX_Roughness_Anisotropy'"), 
 															    TEXT("MaterialFunction'/InterchangeAssets/Functions/MX_Roughness_Dual.MX_Roughness_Dual'"),
-																TEXT("MaterialFunction'/Interchange/Functions/MX_Place2D.MX_Place2D'"),
+																TEXT("MaterialFunction'/InterchangeAssets/Functions/MX_Place2D.MX_Place2D'"),
 															  });
 
 		return bPackagesLoaded;
