Update to handle thin translucent shading model

diff --git a/Engine/Plugins/Interchange/Assets/Shaders/ThinFilmFunctions.usf b/Engine/Plugins/Interchange/Assets/Shaders/ThinFilmFunctions.usf
--- a/Engine/Plugins/Interchange/Assets/Shaders/ThinFilmFunctions.usf
+++ b/Engine/Plugins/Interchange/Assets/Shaders/ThinFilmFunctions.usf
@@ -1,5 +1,7 @@
 // Copyright Epic Games, Inc. All Rights Reserved.
 
+#pragma once
+
 float3 ThinFilmEvalSensitivity(float opd, float shift)
 {
 	// Use Gaussian fits, given by 3 parameters: val, pos and var
