Update function names to avoid substrate conflict.

diff --git a/Engine/Plugins/Interchange/Assets/Shaders/ThinFilmFunctions.usf b/Engine/Plugins/Interchange/Assets/Shaders/ThinFilmFunctions.usf
--- a/Engine/Plugins/Interchange/Assets/Shaders/ThinFilmFunctions.usf
+++ b/Engine/Plugins/Interchange/Assets/Shaders/ThinFilmFunctions.usf
@@ -2,7 +2,7 @@
 
 #pragma once
 
-float3 ThinFilmEvalSensitivity(float opd, float shift)
+float3 ThinFilmEvalSensitivityA(float opd, float shift)
 {
 	// Use Gaussian fits, given by 3 parameters: val, pos and var
 	float phase = 2 * PI * opd * 1.0e-6;
@@ -14,7 +14,7 @@ float3 ThinFilmEvalSensitivity(float opd, float shift)
 	return xyz / 1.0685e-7;
 }
 
-void ThinFilmFresnelDielectric(
+void ThinFilmFresnelDielectricA(
 	in float ct1,
 	in float n1,
 	in float n2,
@@ -44,7 +44,7 @@ void ThinFilmFresnelDielectric(
 	}
 }
 
-void ThinFilmFresnelConductor(
+void ThinFilmFresnelConductorA(
 	in float ct1,
 	in float n1,
 	in float n2,
@@ -73,7 +73,7 @@ void ThinFilmFresnelConductor(
 
 // [Gulbrandsen, "Artist Friendly Metallic Fresnel"]
 // Compute complex index of refraction N = n + ik from F0 and edge tint.
-void ComputeComplexIORFromF0AndEdgeTint(
+void ComputeComplexIORFromF0AndEdgeTintA(
     in float3 Reflectivity, 
     in float3 EdgeTint, 
     inout float3 n, 
@@ -89,8 +89,8 @@ void ComputeComplexIORFromF0AndEdgeTint(
 	k = sqrt(k2);
 }
 
-float ThinFilmDepol(float2 polV) { return 0.5 * (polV.x + polV.y); }
-float3 ThinFilmDepolColor(float3 colS, float3 colP) { return 0.5 * (colS + colP); }
+float ThinFilmDepolA(float2 polV) { return 0.5 * (polV.x + polV.y); }
+float3 ThinFilmDepolColorA(float3 colS, float3 colP) { return 0.5 * (colS + colP); }
 
 float3 SchlickB(float3 F0, float3 F90, float VoH)
 {
@@ -98,7 +98,7 @@ float3 SchlickB(float3 F0, float3 F90, float VoH)
 	return F90 * Fc + (1 - Fc) * F0;
 }
 
-float3 Schlick( float3 SpecularColor, float VoH )
+float3 SchlickA( float3 SpecularColor, float VoH )
 {
 	float Fc = Pow5( 1 - VoH );					// 1 sub, 3 mul
 	//return Fc + (1 - Fc) * SpecularColor;		// 1 add, 3 mad
@@ -107,7 +107,7 @@ float3 Schlick( float3 SpecularColor, float VoH )
 	return saturate( 50.0 * SpecularColor.g ) * Fc + (1 - Fc) * SpecularColor;
 }
 
-float DielectricToF0(float Ior)
+float DielectricToF0A(float Ior)
 {
 	const float F0Sqrt = (Ior-1)/(Ior+1);
 	const float F0 = F0Sqrt*F0Sqrt;
