Interchange Shaders:
- Fix a wrong call causing a compiling error in the Material Editor

#tests test_material_eggs

diff --git a/Engine/Plugins/Interchange/Assets/Shaders/ThinFilmFunctions.usf b/Engine/Plugins/Interchange/Assets/Shaders/ThinFilmFunctions.usf
--- a/Engine/Plugins/Interchange/Assets/Shaders/ThinFilmFunctions.usf
+++ b/Engine/Plugins/Interchange/Assets/Shaders/ThinFilmFunctions.usf
@@ -54,7 +54,7 @@ void ThinFilmFresnelConductorA(
 {
 	if (k == 0)
 	{ // use dielectric formula to avoid numerical issues
-		ThinFilmFresnelDielectric(ct1, n1, n2, R, phi);
+		ThinFilmFresnelDielectricA(ct1, n1, n2, R, phi);
 	}
 	else
 	{
